Editgrating fix for MessageEndpoint hangs from Fury branch.

Fixes issue with Mac editor hanging during loading which seems to have become more prevalent in 10.15 / Xcode 11.

diff --git a/Engine/Source/Runtime/MessagingCommon/Public/MessageEndpoint.h b/Engine/Source/Runtime/MessagingCommon/Public/MessageEndpoint.h
--- a/Engine/Source/Runtime/MessagingCommon/Public/MessageEndpoint.h
+++ b/Engine/Source/Runtime/MessagingCommon/Public/MessageEndpoint.h
@@ -101,7 +101,7 @@ public:
 
 	/** Destructor. */
 	~FMessageEndpoint()
-	{
+	{	
 		auto Bus = BusPtr.Pin();
 
 		if (Bus.IsValid())
@@ -919,24 +919,30 @@ public:
 	 *
 	 * When an object that owns a message endpoint receiving on AnyThread is being destroyed,
 	 * it is possible that the endpoint can outlive the object for a brief period of time if
-	 * the Messaging system is dispatching messages to it. The purpose of this helper function
+	 * the Messaging system is dispatching messages to it. This helper function
 	 * is to block the calling thread while any messages are being dispatched, so that the
 	 * endpoint does not invoke any message handlers after the object has been destroyed.
 	 *
-	 * Note: When calling this function make sure that no other object is holding on to
-	 * the endpoint, or otherwise the caller may get blocked forever.
-	 *
 	 * @param Endpoint The message endpoint to release.
 	 */
 	static void SafeRelease(TSharedPtr<FMessageEndpoint, ESPMode::ThreadSafe>& Endpoint)
 	{
-		TWeakPtr<FMessageEndpoint, ESPMode::ThreadSafe> EndpointPtr = Endpoint;
+		Endpoint->ClearHandlers();
 		Endpoint.Reset();
-		while (EndpointPtr.IsValid());
 	}
 
 protected:
 
+	/**
+	 * Clears all handlers in a way that guarantees it won't overlap with message processing. This preserves internal integrity
+	 * of the array and cases where our owner may be shutting down while receiving messages.
+	*/
+	void ClearHandlers()
+	{
+		FScopeLock Lock(&HandlersCS);
+		Handlers.Empty();
+	}
+
 	/**
 	 * Gets a shared pointer to the message bus if this endpoint is enabled.
 	 *
@@ -964,6 +970,8 @@ protected:
 			return;
 		}
 
+		FScopeLock Lock(&HandlersCS);
+
 		for (int32 HandlerIndex = 0; HandlerIndex < Handlers.Num(); ++HandlerIndex)
 		{
 			Handlers[HandlerIndex]->HandleMessage(Context);
@@ -1006,6 +1014,9 @@ private:
 
 	/** Holds a delegate that is invoked in case of messaging errors. */
 	FOnMessageEndpointError ErrorDelegate;
+
+	/** Sigfnifies that the handler array is being accessed and other threads should wait or skip */
+	FCriticalSection		HandlersCS;
 };
 
 
