Fixed non-unity compilation issues in crashreporter and UFE

diff --git a/Engine/Source/Runtime/MessagingCommon/Public/MessageEndpoint.h b/Engine/Source/Runtime/MessagingCommon/Public/MessageEndpoint.h
--- a/Engine/Source/Runtime/MessagingCommon/Public/MessageEndpoint.h
+++ b/Engine/Source/Runtime/MessagingCommon/Public/MessageEndpoint.h
@@ -17,6 +17,7 @@
 #include "Misc/Guid.h"
 #include "Templates/SharedPointer.h"
 #include "UObject/NameTypes.h"
+#include "Misc/ScopeLock.h"
 
 
 /**
