Fixed missing pointer check in FMessageEndpoint::SafeRelease

#6694

diff --git a/Engine/Source/Runtime/MessagingCommon/Public/MessageEndpoint.h b/Engine/Source/Runtime/MessagingCommon/Public/MessageEndpoint.h
--- a/Engine/Source/Runtime/MessagingCommon/Public/MessageEndpoint.h
+++ b/Engine/Source/Runtime/MessagingCommon/Public/MessageEndpoint.h
@@ -928,8 +928,11 @@ public:
 	 */
 	static void SafeRelease(TSharedPtr<FMessageEndpoint, ESPMode::ThreadSafe>& Endpoint)
 	{
-		Endpoint->ClearHandlers();
-		Endpoint.Reset();
+		if (Endpoint.IsValid())
+		{
+			Endpoint->ClearHandlers();
+			Endpoint.Reset();
+		}
 	}
 
 protected:
